fixed up incorrect module dependencies in DeviceProfileServices module; cosmetic changes

diff --git a/Engine/Source/Editor/DeviceProfileServices/DeviceProfileServices.Build.cs b/Engine/Source/Editor/DeviceProfileServices/DeviceProfileServices.Build.cs
--- a/Engine/Source/Editor/DeviceProfileServices/DeviceProfileServices.Build.cs
+++ b/Engine/Source/Editor/DeviceProfileServices/DeviceProfileServices.Build.cs
@@ -14,11 +14,16 @@ public class DeviceProfileServices : ModuleRules
 				"Core",
 				"CoreUObject", // @todo Mac: for some reason it's needed to link in debug on Mac
 				"Engine",
+				"TargetPlatform",
 				"UnrealEd",
-				"TargetDeviceServices",
 			}
 		);
 
+		PrivateIncludePathModuleNames.AddRange(
+			new string[] {
+				"TargetDeviceServices",
+			}
+		);
 	}
 }
 
diff --git a/Engine/Source/Editor/DeviceProfileServices/Private/DeviceProfileServicesPCH.h b/Engine/Source/Editor/DeviceProfileServices/Private/DeviceProfileServicesPCH.h
--- a/Engine/Source/Editor/DeviceProfileServices/Private/DeviceProfileServicesPCH.h
+++ b/Engine/Source/Editor/DeviceProfileServices/Private/DeviceProfileServicesPCH.h
@@ -6,14 +6,19 @@
 
 #pragma once
 
+#include "DeviceProfileServices.h"
+
+
 /* Dependencies
  *****************************************************************************/
+
 #include "EnginePrivate.h"
-#include "Core.h"
 #include "ModuleManager.h"
+#include "TargetDeviceServices.h"
 
 
 /* Private includes
 *****************************************************************************/
+
 #include "DeviceProfileServices.h"
 #include "DeviceProfileServicesUIManager.h"
diff --git a/Engine/Source/Editor/DeviceProfileServices/Private/DeviceProfileServicesUIManager.cpp b/Engine/Source/Editor/DeviceProfileServices/Private/DeviceProfileServicesUIManager.cpp
--- a/Engine/Source/Editor/DeviceProfileServices/Private/DeviceProfileServicesUIManager.cpp
+++ b/Engine/Source/Editor/DeviceProfileServices/Private/DeviceProfileServicesUIManager.cpp
@@ -5,10 +5,11 @@
 =============================================================================*/
 
 #include "DeviceProfileServicesPCH.h"
-#include "TargetDeviceServices.h"
+
 
 DEFINE_LOG_CATEGORY_STATIC(LogDeviceProfileServices, Log, All);
 
+
 FDeviceProfileServicesUIManager::FDeviceProfileServicesUIManager()
 {
 	GEngine->GetDeviceProfileManager()->OnManagerUpdated().AddRaw(this, &FDeviceProfileServicesUIManager::HandleRefreshUIData);
diff --git a/Engine/Source/Editor/DeviceProfileServices/Private/DeviceProfileServicesUIManager.h b/Engine/Source/Editor/DeviceProfileServices/Private/DeviceProfileServicesUIManager.h
--- a/Engine/Source/Editor/DeviceProfileServices/Private/DeviceProfileServicesUIManager.h
+++ b/Engine/Source/Editor/DeviceProfileServices/Private/DeviceProfileServicesUIManager.h
@@ -20,14 +20,17 @@ public:
 	 */
 	FDeviceProfileServicesUIManager( );
 
+public:
+
 	// Begin IDeviceProfileServicesUIManager Interface
+
 	virtual const FName GetDeviceIconName( const FString& DeviceName ) const OVERRIDE;
 	virtual const TArray<TSharedPtr<FString> > GetPlatformList( ) OVERRIDE;
 	virtual void GetProfilesByType( TArray<UDeviceProfile*>& OutDeviceProfiles, const FString& InType ) OVERRIDE;
 	virtual const FName GetPlatformIconName( const FString& DeviceName ) const OVERRIDE;
 	virtual void SetProfile( const FString& DeviceProfileName ) OVERRIDE;
-	// End IDeviceProfileServicesUIManager Interface
 
+	// End IDeviceProfileServicesUIManager Interface
 
 protected:
 
@@ -41,7 +44,6 @@ protected:
 	 */
 	void HandleRefreshUIData();
 
-
 private:
 
 	// Map of profiles to platform types
@@ -55,4 +57,4 @@ private:
 
 	// Holds the list of known platforms.
 	TArray<TSharedPtr<FString> > PlatformList;
-};
\ No newline at end of file
+};
diff --git a/Engine/Source/Editor/DeviceProfileServices/Public/IDeviceProfileServicesUIManager.h b/Engine/Source/Editor/DeviceProfileServices/Public/IDeviceProfileServicesUIManager.h
--- a/Engine/Source/Editor/DeviceProfileServices/Public/IDeviceProfileServicesUIManager.h
+++ b/Engine/Source/Editor/DeviceProfileServices/Public/IDeviceProfileServicesUIManager.h
@@ -33,7 +33,6 @@ public:
 	 */
 	virtual const FName GetDeviceIconName( const FString& DeviceName ) const = 0;
 
-
 	/**
 	 * Gets Platform List.
 	 *
@@ -41,7 +40,6 @@ public:
 	 */
 	virtual const TArray<TSharedPtr<FString> > GetPlatformList( ) = 0;
 
-
 	/**
 	 * Gets an array of profile that match a type.
 	 *
@@ -50,7 +48,6 @@ public:
 	 */
 	virtual void GetProfilesByType( TArray<UDeviceProfile*>& OutDeviceProfiles, const FString& InType ) = 0;
 
-
 	/**
 	 * Gets an icon name for a given platform.
 	 *
@@ -59,11 +56,10 @@ public:
 	 */
 	virtual const FName GetPlatformIconName( const FString& DeviceName ) const = 0;
 
-
 	/**
-	 * Set a profile - add the type to a .ini file so it can be used in a history view.
+	 * Set a profile - add the type to an INI file so it can be used in a history view.
 	 *
 	 * @param DeviceProfileName - The platform type.
 	 */
 	virtual void SetProfile( const FString& DeviceProfileName ) = 0;
-};
\ No newline at end of file
+};
