fixing incorrect module dependencies in DeviceProfileServices

diff --git a/Engine/Source/Editor/DeviceProfileServices/DeviceProfileServices.Build.cs b/Engine/Source/Editor/DeviceProfileServices/DeviceProfileServices.Build.cs
--- a/Engine/Source/Editor/DeviceProfileServices/DeviceProfileServices.Build.cs
+++ b/Engine/Source/Editor/DeviceProfileServices/DeviceProfileServices.Build.cs
@@ -9,11 +9,17 @@ public class DeviceProfileServices : ModuleRules
 	{
 		PrivateIncludePaths.Add("Editor/DeviceProfileServices/Private");
 
-		PrivateDependencyModuleNames.AddRange(
+		PublicDependencyModuleNames.AddRange(
 			new string[] {
 				"Core",
-				"CoreUObject", // @todo Mac: for some reason it's needed to link in debug on Mac
 				"Engine",
+			}
+		);
+
+		PrivateDependencyModuleNames.AddRange(
+			new string[] {
+				
+				"CoreUObject",
 				"TargetPlatform",
 				"UnrealEd",
 			}
diff --git a/Engine/Source/Editor/DeviceProfileServices/Public/DeviceProfileServices.h b/Engine/Source/Editor/DeviceProfileServices/Public/DeviceProfileServices.h
--- a/Engine/Source/Editor/DeviceProfileServices/Public/DeviceProfileServices.h
+++ b/Engine/Source/Editor/DeviceProfileServices/Public/DeviceProfileServices.h
@@ -10,10 +10,12 @@
  *****************************************************************************/
 
 #include "Core.h"
+#include "Engine.h"
 #include "ModuleInterface.h"
 
+
 /* Interfaces
-//  *****************************************************************************/
+ *****************************************************************************/
 
 #include "IDeviceProfileServicesUIManager.h"
 #include "IDeviceProfileServicesModule.h"
