Add a commandlet to simulate android device profile selection from properties stored in a device.json file.
AndroidDeviceProfileSelectorModule is configured with values from a device.json file and used during device profile selection process.
the results are saved out to multiple files.

diff --git a/Engine/Config/Android/AndroidEngine.ini b/Engine/Config/Android/AndroidEngine.ini
--- a/Engine/Config/Android/AndroidEngine.ini
+++ b/Engine/Config/Android/AndroidEngine.ini
@@ -33,6 +33,7 @@ SmallestMemoryBucket_MinGB=3
 
 [DeviceProfileManager]
 DeviceProfileSelectionModule="AndroidDeviceProfileSelectorRuntime"
+PreviewDeviceProfileSelectionModule="AndroidDeviceProfileSelector"
 
 [Advertising]
 DefaultProviderName=AndroidAdvertising
diff --git a/Engine/Plugins/Runtime/AndroidDeviceProfileSelector/Source/AndroidDeviceProfileSelector/Private/AndroidDeviceProfileSelectorModule.cpp b/Engine/Plugins/Runtime/AndroidDeviceProfileSelector/Source/AndroidDeviceProfileSelector/Private/AndroidDeviceProfileSelectorModule.cpp
--- a/Engine/Plugins/Runtime/AndroidDeviceProfileSelector/Source/AndroidDeviceProfileSelector/Private/AndroidDeviceProfileSelectorModule.cpp
+++ b/Engine/Plugins/Runtime/AndroidDeviceProfileSelector/Source/AndroidDeviceProfileSelector/Private/AndroidDeviceProfileSelectorModule.cpp
@@ -53,6 +53,15 @@ bool FAndroidDeviceProfileSelectorModule::GetSelectorPropertyValue(const FName&
 		PropertyValueOUT = *Found;
 		return true;
 	}
+	// Special case for non-existent config rule variables
+	// they should return true and a value of '[null]'
+	// this prevents configrule issues from throwing errors.
+	if (PropertyType.ToString().StartsWith(TEXT("SRC_ConfigRuleVar[")))
+	{
+		PropertyValueOUT = TEXT("[null]");
+		return true;
+	}
+
 	return false;
 }
 
diff --git a/Engine/Source/Editor/DeviceProfileServices/DeviceProfileServices.Build.cs b/Engine/Source/Editor/DeviceProfileServices/DeviceProfileServices.Build.cs
--- a/Engine/Source/Editor/DeviceProfileServices/DeviceProfileServices.Build.cs
+++ b/Engine/Source/Editor/DeviceProfileServices/DeviceProfileServices.Build.cs
@@ -13,6 +13,8 @@ public class DeviceProfileServices : ModuleRules
 			new string[] {
 				"Core",
 				"Engine",
+				"Json",
+				"JsonUtilities",
 			});
 
 		PrivateDependencyModuleNames.AddRange(
@@ -22,6 +24,7 @@ public class DeviceProfileServices : ModuleRules
 				"DesktopPlatform",
 				"EditorFramework",
 				"UnrealEd",
+				"PIEPreviewDeviceSpecification",
 			});
 
 		PrivateIncludePathModuleNames.AddRange(
diff --git a/Engine/Source/Editor/DeviceProfileServices/Private/CheckAndroidDeviceProfileCommandlet.cpp b/Engine/Source/Editor/DeviceProfileServices/Private/CheckAndroidDeviceProfileCommandlet.cpp
--- a/Engine/Source/Editor/DeviceProfileServices/Private/CheckAndroidDeviceProfileCommandlet.cpp
+++ b/Engine/Source/Editor/DeviceProfileServices/Private/CheckAndroidDeviceProfileCommandlet.cpp
@@ -3,9 +3,67 @@
 #include "CheckAndroidDeviceProfileCommandlet.h"
 #include "IDeviceProfileSelectorModule.h"
 #include "Modules/ModuleManager.h"
+#include "DeviceProfiles/DeviceProfileManager.h"
+#include "PIEPreviewDeviceSpecification.h"
+#include "JsonObjectConverter.h"
+#include "Misc/FileHelper.h"
+#include "Misc/ConfigCacheIni.h"
+#include "DeviceProfiles/DeviceProfile.h"
 
 DEFINE_LOG_CATEGORY_STATIC(LogCheckAndroidDeviceProfile, Log, All);
 
+EPlatformMemorySizeBucket DetermineDeviceMemorySizeBucket(int32 TotalPhysicalMemoryGB)
+{
+	EPlatformMemorySizeBucket Bucket = EPlatformMemorySizeBucket::Default;
+	static int32 LargestMemoryGB = 0, LargerMemoryGB = 0, DefaultMemoryGB = 0, SmallerMemoryGB = 0, SmallestMemoryGB = 0, TiniestMemoryGB = 0;
+	static bool bTriedPlatformBuckets = false;
+	if (!bTriedPlatformBuckets)
+	{
+		bTriedPlatformBuckets = true;
+		FConfigFile EngineConfigFile;
+		if (FConfigCacheIni::LoadLocalIniFile(EngineConfigFile, TEXT("Engine"), true, TEXT("Android")))
+		{
+			// get values for this platform from it's .ini
+			EngineConfigFile.GetInt(TEXT("PlatformMemoryBuckets"), TEXT("LargestMemoryBucket_MinGB"), LargestMemoryGB);
+			EngineConfigFile.GetInt(TEXT("PlatformMemoryBuckets"), TEXT("LargerMemoryBucket_MinGB"), LargerMemoryGB);
+			EngineConfigFile.GetInt(TEXT("PlatformMemoryBuckets"), TEXT("DefaultMemoryBucket_MinGB"), DefaultMemoryGB);
+			EngineConfigFile.GetInt(TEXT("PlatformMemoryBuckets"), TEXT("SmallerMemoryBucket_MinGB"), SmallerMemoryGB);
+			EngineConfigFile.GetInt(TEXT("PlatformMemoryBuckets"), TEXT("SmallestMemoryBucket_MinGB"), SmallestMemoryGB);
+			EngineConfigFile.GetInt(TEXT("PlatformMemoryBuckets"), TEXT("TiniestMemoryBucket_MinGB"), TiniestMemoryGB);
+		}
+	}
+
+	// if at least Smaller is specified, we can set the Bucket
+	if (SmallerMemoryGB > 0)
+	{
+		if (TotalPhysicalMemoryGB >= SmallerMemoryGB)
+		{
+			Bucket = EPlatformMemorySizeBucket::Smaller;
+		}
+		else if (TotalPhysicalMemoryGB >= SmallestMemoryGB)
+		{
+			Bucket = EPlatformMemorySizeBucket::Smallest;
+		}
+		else
+		{
+			Bucket = EPlatformMemorySizeBucket::Tiniest;
+		}
+	}
+	if (DefaultMemoryGB > 0 && TotalPhysicalMemoryGB >= DefaultMemoryGB)
+	{
+		Bucket = EPlatformMemorySizeBucket::Default;
+	}
+	if (LargerMemoryGB > 0 && TotalPhysicalMemoryGB >= LargerMemoryGB)
+	{
+		Bucket = EPlatformMemorySizeBucket::Larger;
+	}
+	if (LargestMemoryGB > 0 && TotalPhysicalMemoryGB >= LargestMemoryGB)
+	{
+		Bucket = EPlatformMemorySizeBucket::Largest;
+	}
+	return Bucket;
+}
+
 int32 UCheckAndroidDeviceProfileCommandlet::Main(const FString& RawCommandLine)
 {
 	TArray<FString> Tokens;
@@ -13,34 +71,177 @@ int32 UCheckAndroidDeviceProfileCommandlet::Main(const FString& RawCommandLine)
 	TMap<FString, FString> Params;
 	ParseCommandLine(*RawCommandLine, Tokens, Switches, Params);
 
-	IDeviceProfileSelectorModule* AndroidDeviceProfileSelector = FModuleManager::LoadModulePtr<IDeviceProfileSelectorModule>("AndroidDeviceProfileSelector");
+	FString ParamDeviceSpecsFolder = Params.FindRef(TEXT("DeviceSpecsFolder"));
+	FString ParamDeviceSpecsFile = Params.FindRef(TEXT("DeviceSpecsFile"));
+	FString ForcedDeviceProfileName = Params.FindRef(TEXT("OverrideDP"));
+	FString OutputFolder = Params.FindRef(TEXT("OutDir"));
+
+	TArray<FString> DeviceSpecificationFileNames;
+	if (!ParamDeviceSpecsFolder.IsEmpty())
+	{
+		IFileManager::Get().FindFiles(DeviceSpecificationFileNames, *(ParamDeviceSpecsFolder / TEXT("*.json")), true, false);
+	}
+	else if (!ParamDeviceSpecsFile.IsEmpty())
+	{
+		DeviceSpecificationFileNames.Add(ParamDeviceSpecsFile);
+	}
+
+	FName DeviceProfileSelectorModule("AndroidDeviceProfileSelector");
+	IDeviceProfileSelectorModule* AndroidDeviceProfileSelector = FModuleManager::LoadModulePtr<IDeviceProfileSelectorModule>(DeviceProfileSelectorModule);
 	if (ensure(AndroidDeviceProfileSelector != nullptr))
 	{
-		TMap<FName, FString> DeviceParameters;
-		DeviceParameters.Add(FName(TEXT("SRC_GPUFamily")), Params.FindRef(TEXT("GPUFamily")));
-		DeviceParameters.Add(FName(TEXT("SRC_GLVersion")), Params.FindRef(TEXT("GLVersion")));
-		DeviceParameters.Add(FName(TEXT("SRC_VulkanAvailable")), Params.FindRef(TEXT("VulkanAvailable")));
-		DeviceParameters.Add(FName(TEXT("SRC_VulkanVersion")), Params.FindRef(TEXT("VulkanVersion")));
-		DeviceParameters.Add(FName(TEXT("SRC_AndroidVersion")), Params.FindRef(TEXT("AndroidVersion")));
-		DeviceParameters.Add(FName(TEXT("SRC_DeviceMake")),
-			Tokens.Num() == 2 ? Tokens[0] :
-			Params.FindRef(TEXT("DeviceMake")));
-		DeviceParameters.Add(FName(TEXT("SRC_DeviceModel")),
-			Tokens.Num() == 1 ? Tokens[0] :
-			Tokens.Num() == 2 ? Tokens[1] : 
-			Params.FindRef(TEXT("DeviceModel")));
-		DeviceParameters.Add(FName(TEXT("SRC_DeviceBuildNumber")), Params.FindRef(TEXT("DeviceBuildNumber")));
-		DeviceParameters.Add(FName(TEXT("SRC_UsingHoudini")), Params.FindRef(TEXT("UsingHoudini")));
-		DeviceParameters.Add(FName(TEXT("SRC_Hardware")), Params.FindRef(TEXT("Hardware")));
-		DeviceParameters.Add(FName(TEXT("SRC_Chipset")), Params.FindRef(TEXT("Chipset")));
-		DeviceParameters.Add(FName(TEXT("SRC_HMDSystemName")), Params.FindRef(TEXT("HMDSystemName")));
-		DeviceParameters.Add(FName(TEXT("SRC_TotalPhysicalGB")), Params.FindRef(TEXT("TotalPhysicalGB")));
-
-		AndroidDeviceProfileSelector->SetSelectorProperties(DeviceParameters);
-		FString ProfileName = AndroidDeviceProfileSelector->GetDeviceProfileName();
-
-		UE_LOG(LogCheckAndroidDeviceProfile, Display, TEXT("Selected Device Profile: %s"), *ProfileName);
+		do
+		{
+			TMap<FName, FString> DeviceParameters;
+			FString DeviceName;
+			FString OutputFilename;
+			if (DeviceSpecificationFileNames.Num())
+			{
+				FString DeviceSpecsFile = DeviceSpecificationFileNames.Pop();
+				int32 ExtensionPos;
+				if (DeviceSpecsFile.FindLastChar(TEXT('.'), ExtensionPos))
+				{
+					DeviceName = DeviceSpecsFile.Left(ExtensionPos);
+					OutputFilename = DeviceName;
+				}
+				TSharedPtr<FJsonObject> JsonRootObject;
+				FString Json;
+
+				if (FFileHelper::LoadFileToString(Json, *(ParamDeviceSpecsFolder / DeviceSpecsFile)))
+				{
+					TSharedRef<TJsonReader<> > JsonReader = TJsonReaderFactory<>::Create(Json);
+					FJsonSerializer::Deserialize(JsonReader, JsonRootObject);
+				}
+				FPIEPreviewDeviceSpecifications DeviceSpecs;
+				if (JsonRootObject.IsValid() && FJsonObjectConverter::JsonAttributesToUStruct(JsonRootObject->Values, FPIEPreviewDeviceSpecifications::StaticStruct(), &DeviceSpecs, 0, 0))
+				{
+					FPIEAndroidDeviceProperties& AndroidProperties = DeviceSpecs.AndroidProperties;
+					DeviceParameters.Add(FName(TEXT("SRC_GPUFamily")), AndroidProperties.GPUFamily);
+					DeviceParameters.Add(FName(TEXT("SRC_GLVersion")), AndroidProperties.GLVersion);
+					DeviceParameters.Add(FName(TEXT("SRC_VulkanAvailable")), AndroidProperties.VulkanAvailable ? "true" : "false");
+					DeviceParameters.Add(FName(TEXT("SRC_VulkanVersion")), AndroidProperties.VulkanVersion);
+					DeviceParameters.Add(FName(TEXT("SRC_AndroidVersion")), AndroidProperties.AndroidVersion);
+					DeviceParameters.Add(FName(TEXT("SRC_DeviceMake")), AndroidProperties.DeviceMake);
+					DeviceParameters.Add(FName(TEXT("SRC_DeviceModel")), AndroidProperties.DeviceModel);
+					DeviceParameters.Add(FName(TEXT("SRC_DeviceBuildNumber")), AndroidProperties.DeviceBuildNumber);
+					DeviceParameters.Add(FName(TEXT("SRC_UsingHoudini")), AndroidProperties.UsingHoudini ? "true" : "false");
+					DeviceParameters.Add(FName(TEXT("SRC_Hardware")), AndroidProperties.Hardware);
+					DeviceParameters.Add(FName(TEXT("SRC_Chipset")), AndroidProperties.Chipset);
+					DeviceParameters.Add(FName(TEXT("SRC_TotalPhysicalGB")), AndroidProperties.TotalPhysicalGB);
+					DeviceParameters.Add(FName(TEXT("SRC_HMDSystemName")), TEXT(""));
+				}
+			}
+			else
+			{
+				DeviceName = TEXT("[no name]");
+				DeviceParameters.Add(FName(TEXT("SRC_GPUFamily")), Params.FindRef(TEXT("GPUFamily")));
+				DeviceParameters.Add(FName(TEXT("SRC_GLVersion")), Params.FindRef(TEXT("GLVersion")));
+				DeviceParameters.Add(FName(TEXT("SRC_VulkanAvailable")), Params.FindRef(TEXT("VulkanAvailable")));
+				DeviceParameters.Add(FName(TEXT("SRC_VulkanVersion")), Params.FindRef(TEXT("VulkanVersion")));
+				DeviceParameters.Add(FName(TEXT("SRC_AndroidVersion")), Params.FindRef(TEXT("AndroidVersion")));
+				DeviceParameters.Add(FName(TEXT("SRC_DeviceMake")),
+					Tokens.Num() == 2 ? Tokens[0] :
+					Params.FindRef(TEXT("DeviceMake")));
+				DeviceParameters.Add(FName(TEXT("SRC_DeviceModel")),
+					Tokens.Num() == 1 ? Tokens[0] :
+					Tokens.Num() == 2 ? Tokens[1] :
+					Params.FindRef(TEXT("DeviceModel")));
+				DeviceParameters.Add(FName(TEXT("SRC_DeviceBuildNumber")), Params.FindRef(TEXT("DeviceBuildNumber")));
+				DeviceParameters.Add(FName(TEXT("SRC_UsingHoudini")), Params.FindRef(TEXT("UsingHoudini")));
+				DeviceParameters.Add(FName(TEXT("SRC_Hardware")), Params.FindRef(TEXT("Hardware")));
+				DeviceParameters.Add(FName(TEXT("SRC_Chipset")), Params.FindRef(TEXT("Chipset")));
+				DeviceParameters.Add(FName(TEXT("SRC_TotalPhysicalGB")), Params.FindRef(TEXT("TotalPhysicalGB")));
+				DeviceParameters.Add(FName(TEXT("SRC_HMDSystemName")), TEXT(""));
+			}
+			AndroidDeviceProfileSelector->SetSelectorProperties(DeviceParameters);
+			FString ProfileName = ForcedDeviceProfileName.IsEmpty() ? AndroidDeviceProfileSelector->GetDeviceProfileName() : ForcedDeviceProfileName;
+			TArray<FSelectedFragmentProperties> SelectedFragments;
+
+			UE_LOG(LogCheckAndroidDeviceProfile, Display, TEXT("%s Selected Device Profile: %s"), *DeviceName, *ProfileName);
+			int32 TotalPhysGB = FCString::Atoi64(*DeviceParameters.FindChecked(TEXT("SRC_TotalPhysicalGB")));
+			EPlatformMemorySizeBucket MemorySizeBucket = DetermineDeviceMemorySizeBucket(TotalPhysGB);
+			FString PlatformOverride(TEXT("Android"));
+			FString ProfileDescription;
+			UDeviceProfile* Profile = UDeviceProfileManager::Get().FindProfile(ProfileName, false);
+			TArray<FString> CVars;
+			if (!Profile)
+			{
+				UE_LOG(LogCheckAndroidDeviceProfile, Error, TEXT("Failed to find requested device profile. %s requested Device Profile: %s But could not be found!"), *DeviceName, *ProfileName);
+			}
+			else
+			{
+				// Set the memory size bucket for this device.
+				Profile->SetPreviewMemorySizeBucket(MemorySizeBucket);
+				// ensure any previous expanded cvars are cleared out.
+				Profile->ClearAllExpandedCVars();
+				// get all GetAllExpandedCVars re-runs the DP selection.
+				TMap<FString, FString> CVarMap = Profile->GetAllExpandedCVars();
+
+				for (auto CVarIt : CVarMap)
+				{
+					CVars.Add(FString::Printf(TEXT("%s=%s"), *CVarIt.Key, *CVarIt.Value));
+				}
+				SelectedFragments = Profile->SelectedFragments;
+			}
+
+			if (!OutputFilename.IsEmpty())
+			{
+				FString output;
+
+				if (!DeviceName.IsEmpty())
+				{
+					output += FString::Printf(TEXT("Device name : %s\n\n"), *DeviceName);
+				}
+				if (!Profile)
+				{
+					output += FString::Printf(TEXT("Failed to find device profile : %s\n\n"), *ProfileName);
+				}
+
+				OutputFilename = DeviceName + TEXT("_") + DeviceParameters.FindChecked(TEXT("SRC_GPUFamily")) + TEXT("_") + DeviceParameters.FindChecked(TEXT("SRC_TotalPhysicalGB")) + TEXT("GB.txt");
+				output += FString::Printf(TEXT("AndroidDeviceProfileSelector Params:\n"));
+				TArray<FName> SortedKeys;
+				DeviceParameters.GenerateKeyArray(SortedKeys);
+				//SortedKeys.Sort();
+				for (FName& ParamKey : SortedKeys)
+				{
+					output += FString::Printf(TEXT("[%s]=%s\n"), *ParamKey.ToString(), *DeviceParameters.FindChecked(ParamKey));
+				}
+				output += FString::Printf(TEXT("\nDevice Mem Bucket: %s\n"), LexToString(MemorySizeBucket));
+
+				output += FString::Printf(TEXT("\nSelected fragments :\n%s"), SelectedFragments.Num() == 0 ? TEXT("-none-\n") : TEXT(""));
+				for (FSelectedFragmentProperties& Fragment : SelectedFragments)
+				{
+					FString FragTag;
+					if (Fragment.Tag != NAME_None)
+					{
+						FragTag = TEXT("[") + Fragment.Tag.ToString() + TEXT("]");
+					}
+					output += FString::Printf(TEXT("%s%s : Enabled=%s\n"), *FragTag, *Fragment.Fragment, Fragment.bEnabled ? TEXT("true") : TEXT("false"));
+				}
+
+				output += FString::Printf(TEXT("\n\nSelected device Profile: %s\nSelected device profile description: %s\n"), *ProfileName, *ProfileDescription);
+
+				// may not want this:
+				const bool bSortCVars = true;
+				if (bSortCVars)
+				{
+					output += FString::Printf(TEXT("Sorted "));
+					CVars.Sort();
+				}
+				FString CVarsOnly = FString::Printf(TEXT("CVars:\n"));
+
+				for (FString& CVar : CVars)
+				{
+					CVarsOnly += FString::Printf(TEXT("%s\n"), *CVar);
+				}
+				output += CVarsOnly;
+				FFileHelper::SaveStringToFile(output, *(OutputFolder / OutputFilename));
+
+				FFileHelper::SaveStringToFile(CVarsOnly, *(OutputFolder / TEXT("CVarsOnly") / OutputFilename));
+			}
+
+		} while (DeviceSpecificationFileNames.Num());
 	}
 
 	return 0;
-}
+}
\ No newline at end of file
diff --git a/Engine/Source/Editor/DeviceProfileServices/Public/CheckAndroidDeviceProfileCommandlet.h b/Engine/Source/Editor/DeviceProfileServices/Public/CheckAndroidDeviceProfileCommandlet.h
--- a/Engine/Source/Editor/DeviceProfileServices/Public/CheckAndroidDeviceProfileCommandlet.h
+++ b/Engine/Source/Editor/DeviceProfileServices/Public/CheckAndroidDeviceProfileCommandlet.h
@@ -16,7 +16,13 @@
  * ProjectName -run=DeviceProfileServices.CheckAndroidDeviceProfile DeviceMake DeviceModel
  * ProjectName -run=DeviceProfileServices.CheckAndroidDeviceProfile [-DeviceMake=...] [-DeviceModel=...] [-GPUFamily=...] [-GLVersion=...]
 	  [-VulkanAvailable=True|False] [-VulkanVersion=...] [-AndroidVersion=...] [-DeviceBuildNumber=...] [-UsingHoudini=True|False] [-Hardware=...] [-Chipset=...]
- */
+ * Running against 
+ * ProjectName -run=DeviceProfileServices.CheckAndroidDeviceProfile
+		-DeviceSpecsFolder=<directory containing device.json files> 
+		-DeviceSpecsFile=<path to a single device.json file>
+		-OutDir=<output directory>
+		[-OverrideDP=<optional DP name to override device profile selection>]
+*/
 UCLASS()
 class UCheckAndroidDeviceProfileCommandlet
 	: public UCommandlet
